Fix non-unity errors
#rnx

diff --git a/Engine/Plugins/Developer/DumpGPUServices/Source/DumpGPUServices/Private/DumpGPUHTTPUploadServices.cpp b/Engine/Plugins/Developer/DumpGPUServices/Source/DumpGPUServices/Private/DumpGPUHTTPUploadServices.cpp
--- a/Engine/Plugins/Developer/DumpGPUServices/Source/DumpGPUServices/Private/DumpGPUHTTPUploadServices.cpp
+++ b/Engine/Plugins/Developer/DumpGPUServices/Source/DumpGPUServices/Private/DumpGPUHTTPUploadServices.cpp
@@ -2,6 +2,7 @@
 
 #include "DumpGPU.h"
 #include "DumpGPUServices.h"
+#include "Async/TaskGraphInterfaces.h"
 #include "Misc/App.h"
 #include "Misc/Paths.h"
 #include "Misc/FileHelper.h"
diff --git a/Engine/Plugins/Runtime/OpenXR/Source/OpenXRHMD/Private/OpenXRHMD.cpp b/Engine/Plugins/Runtime/OpenXR/Source/OpenXRHMD/Private/OpenXRHMD.cpp
--- a/Engine/Plugins/Runtime/OpenXR/Source/OpenXRHMD/Private/OpenXRHMD.cpp
+++ b/Engine/Plugins/Runtime/OpenXR/Source/OpenXRHMD/Private/OpenXRHMD.cpp
@@ -8,6 +8,7 @@
 #include "IOpenXRExtensionPlugin.h"
 
 #include "Misc/App.h"
+#include "Misc/ConfigCache.h"
 #include "Misc/Parse.h"
 #include "Modules/ModuleManager.h"
 #include "EngineGlobals.h"
