Works arround NotForLicensees config files being striped out of staged build for DumpGPU

diff --git a/Engine/Plugins/Developer/DumpGPUServices/Source/DumpGPUServices/DumpGPUServices.Build.cs b/Engine/Plugins/Developer/DumpGPUServices/Source/DumpGPUServices/DumpGPUServices.Build.cs
--- a/Engine/Plugins/Developer/DumpGPUServices/Source/DumpGPUServices/DumpGPUServices.Build.cs
+++ b/Engine/Plugins/Developer/DumpGPUServices/Source/DumpGPUServices/DumpGPUServices.Build.cs
@@ -1,5 +1,7 @@
 // Copyright Epic Games, Inc. All Rights Reserved.
 
+using EpicGames.Core;
+
 namespace UnrealBuildTool.Rules
 {
 	public class DumpGPUServices : ModuleRules
@@ -37,6 +39,15 @@ namespace UnrealBuildTool.Rules
 					}
 					);
 			}
+
+			// Engine's NotForLicensees/ configuration are always striped out of staged builds by default. To get DumpGPU services available on all
+			// internal projects automatically on staged builds,  read the engine wide NotForLicensees' default URL pattern at compile time.
+			ConfigHierarchy Config = ConfigCache.ReadHierarchy(ConfigHierarchyType.Engine, null, Target.Platform);
+			string EngineDefaultUploadURLPattern = "";
+			if (Config.GetString("Rendering.DumpGPUServices", "EngineDefaultUploadURLPattern", out EngineDefaultUploadURLPattern))
+			{
+				PrivateDefinitions.Add("DUMPGPU_SERVICES_DEFAULT_URL_PATTERN=\"" + EngineDefaultUploadURLPattern + "\"");
+			}
 		}
 	}
 }
diff --git a/Engine/Plugins/Developer/DumpGPUServices/Source/DumpGPUServices/Private/DumpGPUServices.cpp b/Engine/Plugins/Developer/DumpGPUServices/Source/DumpGPUServices/Private/DumpGPUServices.cpp
--- a/Engine/Plugins/Developer/DumpGPUServices/Source/DumpGPUServices/Private/DumpGPUServices.cpp
+++ b/Engine/Plugins/Developer/DumpGPUServices/Source/DumpGPUServices/Private/DumpGPUServices.cpp
@@ -7,7 +7,7 @@
 #include "Misc/ConfigCacheIni.h"
 #include "Modules/ModuleManager.h"
 #include "IDumpGPUServices.h"
-
+#include "DumpGPU.h"
 
 class FDumpGPUServices : public IDumpGPUServices
 {
@@ -35,6 +35,14 @@ void FDumpGPUServices::StartupModule()
 	FString UploadURLPattern;
 	GConfig->GetString(TEXT("Rendering.DumpGPUServices"), TEXT("UploadURLPattern"), UploadURLPattern, GEngineIni);
 
+	// If the project has not set UploadURLPattern, uses the engine default that maybe set in a NotForLicensees/ config.
+	#if defined(DUMPGPU_SERVICES_DEFAULT_URL_PATTERN)
+	if (UploadURLPattern.IsEmpty())
+	{
+		UploadURLPattern = TEXT(DUMPGPU_SERVICES_DEFAULT_URL_PATTERN);
+	}
+	#endif
+
 	if (UploadURLPattern.StartsWith(TEXT("http://")))
 	{
 		UploadProvider = CreateHTTPUploadProvider(UploadURLPattern);
